- Updated min Xcode to 15.2

diff --git a/Engine/Config/Apple/Apple_SDK.json b/Engine/Config/Apple/Apple_SDK.json
--- a/Engine/Config/Apple/Apple_SDK.json
+++ b/Engine/Config/Apple/Apple_SDK.json
@@ -1,7 +1,7 @@
 {
 	// Xcode versions
-	"MainVersion": "14.1",
-	"MinVersion": "14.1.0",
+	"MainVersion": "15.2",
+	"MinVersion": "15.2.0",
 	"MaxVersion": "16.9.0",
 	// !!!
 	// NOTE: If you update the MaxVersion, double check the AppleVersionToLLVMVersion array below!!!
