Update Xcode16.3 LLVM version mapping

#rnx

diff --git a/Engine/Config/Apple/Apple_SDK.json b/Engine/Config/Apple/Apple_SDK.json
--- a/Engine/Config/Apple/Apple_SDK.json
+++ b/Engine/Config/Apple/Apple_SDK.json
@@ -19,6 +19,7 @@
 			"14.0.0-14.0.0",
 			"14.0.3-15.0.0",
 			"15.0.0-16.0.0",
-			"16.0.0-17.0.6"
+			"16.0.0-17.0.6",
+			"16.3.0-19.1.4"
 		]
 }
