Add support for Xcode26

diff --git a/Engine/Config/Apple/Apple_SDK.json b/Engine/Config/Apple/Apple_SDK.json
--- a/Engine/Config/Apple/Apple_SDK.json
+++ b/Engine/Config/Apple/Apple_SDK.json
@@ -2,7 +2,7 @@
 	"//1": "Xcode versions:",
 		"MainVersion": "15.2",
 		"MinVersion": "15.2.0",
-		"MaxVersion": "16.9.0",
+		"MaxVersion": "26.9.0",
 	"//2": "!!!",
 	"//3": "NOTE: If you update the MaxVersion, double check the AppleVersionToLLVMVersion array below!!!",
 	"//4": "!!!",
@@ -20,6 +20,7 @@
 			"14.0.3-15.0.0",
 			"15.0.0-16.0.0",
 			"16.0.0-17.0.6",
-			"16.3.0-19.1.4"
+			"16.3.0-19.1.4",
+			"26.0.0-19.1.5"
 		]
 }
