[UBT] Updating autosdk config for apple platforms to point to the new toolchain versions.

#Synced-CL 50408994

diff --git a/Engine/Config/Apple/AppleOpenSource_SDK.json b/Engine/Config/Apple/AppleOpenSource_SDK.json
--- a/Engine/Config/Apple/AppleOpenSource_SDK.json
+++ b/Engine/Config/Apple/AppleOpenSource_SDK.json
@@ -15,11 +15,11 @@
 	"//8": "Format: 'FirstXcodeVersion-vN'. Selects highest matching version.",
 	"//9": "Xcode    | Apple Clang | LLVM    | Prefix | Auto SDK Directory",
 	"//10": "16.0.0+  | 16.0.0      | 17.0.6  | v1     | v1_clang-17.0.6_oss-1350",
-	"//11": "16.3.0+  | 17.0.0      | 19.1.4  | v2     | v2_clang-19.1.4_oss-1350",
-	"//12": "26.0.0+  | 17.0.6      | 19.1.5  | v3     | v3_clang-19.1.6_oss-1350",
+	"//11": "16.3.0+  | 17.0.0      | 19.1.4  | v5     | v5_clang-19.1.4_oss-1350_from_source",
+	"//12": "26.0.0+  | 17.0.6      | 19.1.6  | v6     | v6_clang-19.1.6_oss-1350_from_source",
 	"XcodeToToolchainVersion": [
 		"16.0.0-v1",
-		"16.3.0-v2",
-		"26.0.0-v3"
+		"16.3.0-v5",
+		"26.0.0-v6"
 	]
 }
