[Mac/iOS] Enabling support for toolchain discovery with Xcode 26.4 in UBT

The engine still needs fixes to be able to build with this version. This CL just allows UBT to use the correct clang version information

#Synced-CL 51701759
#Synced-CL 51701759

diff --git a/Engine/Config/Apple/AppleOpenSource_SDK.json b/Engine/Config/Apple/AppleOpenSource_SDK.json
--- a/Engine/Config/Apple/AppleOpenSource_SDK.json
+++ b/Engine/Config/Apple/AppleOpenSource_SDK.json
@@ -17,9 +17,11 @@
 	"//10": "16.0.0+  | 16.0.0      | 17.0.6  | v1     | v1_clang-17.0.6_oss-1350",
 	"//11": "16.3.0+  | 17.0.0      | 19.1.4  | v5     | v5_clang-19.1.4_oss-1350_from_source",
 	"//12": "26.0.0+  | 17.0.6      | 19.1.6  | v6     | v6_clang-19.1.6_oss-1350_from_source",
+	"//13": "26.4.0+  | 21.0.0      | 21.1.6  | v7     | v7_clang-21.1.6_oss-1350_from_source",
 	"XcodeToToolchainVersion": [
 		"16.0.0-v1",
 		"16.3.0-v5",
-		"26.0.0-v6"
+		"26.0.0-v6",
+		"26.4-v7"
 	]
 }
diff --git a/Engine/Config/Apple/Apple_SDK.json b/Engine/Config/Apple/Apple_SDK.json
--- a/Engine/Config/Apple/Apple_SDK.json
+++ b/Engine/Config/Apple/Apple_SDK.json
@@ -21,6 +21,7 @@
 			"15.0.0-16.0.0",
 			"16.0.0-17.0.6",
 			"16.3.0-19.1.4",
-			"26.0.0-19.1.5"
+			"26.0.0-19.1.5",
+			"26.4.0-21.1.6"
 		]
 }
